Created a build.cs file for DistributedBuildInterface so it can be referenced by other modules.

diff --git a/Engine/Source/Developer/DistributedBuildInterface/DistributedBuildInterface.build.cs b/Engine/Source/Developer/DistributedBuildInterface/DistributedBuildInterface.build.cs
new file mode 100644
--- /dev/null
+++ b/Engine/Source/Developer/DistributedBuildInterface/DistributedBuildInterface.build.cs
@@ -0,0 +1,10 @@
+// Copyright Epic Games, Inc. All Rights Reserved.
+
+using UnrealBuildTool;
+
+public class DistributedBuildInterface : ModuleRules
+{
+	public DistributedBuildInterface(ReadOnlyTargetRules Target) : base(Target)
+	{
+	}
+}
