Possible fix for DistributedBuildInterface linker issue.

diff --git a/Engine/Source/Developer/DistributedBuildInterface/DistributedBuildInterface.build.cs b/Engine/Source/Developer/DistributedBuildInterface/DistributedBuildInterface.build.cs
--- a/Engine/Source/Developer/DistributedBuildInterface/DistributedBuildInterface.build.cs
+++ b/Engine/Source/Developer/DistributedBuildInterface/DistributedBuildInterface.build.cs
@@ -6,5 +6,7 @@ public class DistributedBuildInterface : ModuleRules
 {
 	public DistributedBuildInterface(ReadOnlyTargetRules Target) : base(Target)
 	{
+		Type = ModuleType.External;
+		PublicSystemIncludePaths.Add(System.IO.Path.Combine(ModuleDirectory, "Public"));
 	}
 }
