[shaders] UBA fixes - reduce UBA CAS size if using virtual files, only SCW binaries end up in there if enabled; hold job refcount in distributed task to avoid timing-related issues where jobs are cancelled while UBA is in the middle of dispatching them (potentially causing a use-after-free on the shader source memory blobs, owned by the job)

#Synced-CL 52297251

diff --git a/Engine/Plugins/UbaController/Source/Private/UbaJobProcessor.cpp b/Engine/Plugins/UbaController/Source/Private/UbaJobProcessor.cpp
--- a/Engine/Plugins/UbaController/Source/Private/UbaJobProcessor.cpp
+++ b/Engine/Plugins/UbaController/Source/Private/UbaJobProcessor.cpp
@@ -472,7 +472,8 @@ void FUbaJobProcessor::StartUba()
 		ConfigTable_AddValueString(*RootTable,     TC("RootDir"), TCHAR_TO_UBASTRING(*RootDir));
 
 		uba::ConfigTable* StorageTable = Config_AddTable(*Config, TC("Storage"));
-		ConfigTable_AddValueU64(*StorageTable,     TC("CasCapacityBytes"), 32llu * 1024 * 1024 * 1024);
+		const uint64 CasSizeMb = UbaJobProcessorOptions::bUseVirtualFiles ? 512 : 32 * 1024; // CAS is basically only populated with SCW and SCW-related binaries when using virtual files, so CAS can be much smaller.
+		ConfigTable_AddValueU64(*StorageTable,     TC("CasCapacityBytes"), CasSizeMb * 1024 * 1024);
 
 		uba::ConfigTable* SessionTable = Config_AddTable(*Config, TC("Session"));
 		ConfigTable_AddValueBool(*SessionTable,    TC("LaunchVisualizer"), UbaJobProcessorOptions::bAutoLaunchVisualizer);
diff --git a/Engine/Source/Developer/DistributedBuildInterface/Public/DistributedBuildControllerInterface.h b/Engine/Source/Developer/DistributedBuildInterface/Public/DistributedBuildControllerInterface.h
--- a/Engine/Source/Developer/DistributedBuildInterface/Public/DistributedBuildControllerInterface.h
+++ b/Engine/Source/Developer/DistributedBuildInterface/Public/DistributedBuildControllerInterface.h
@@ -28,13 +28,12 @@ struct FTaskCommandData
 	FString InputFileName;
 	FString OutputFileName;
 	FString ExtraCommandArgs;
-	FString Description; 						// Optional string describing the task; can be used by the controller implementation for debug purposes (i.e. UBA breadcrumbs)
+	FString Description; 									// Optional string describing the task; can be used by the controller implementation for debug purposes (i.e. UBA breadcrumbs)
 	uint32 DispatcherPID = 0;
-	TArray<uint8> InputFileData; 				// Optional byte array containing the contents of the input file; only used if the controller supports virtual files
-	TArray <FMemoryView> SourceData; 			// Optional array of memory views containing the compressed source files; only used if the controller supports virtual files 
-												// (note that these are populated from the FShaderSource objects owned by FShaderPreprocessOutput, and so this must remain alive 
-												// for the duration of the build task)
-	TArray<FString> AdditionalOutputFolders; 	// Optional additional folder(s) which task may write artifacts to
+	TArray<uint8> InputFileData; 							// Optional byte array containing the contents of the input file; only used if the controller supports virtual files
+	TArray <FMemoryView> SourceData; 						// Optional array of memory views containing the compressed source files; only used if the controller supports virtual files
+	TSharedPtr<void, ESPMode::ThreadSafe> SourceDataOwner;	// Optional shared pointer to opaque data keeping the memory referenced by SourceData alive
+	TArray<FString> AdditionalOutputFolders; 				// Optional additional folder(s) which task may write artifacts to
 };
 
 struct FDistributedBuildTask
diff --git a/Engine/Source/Runtime/Engine/Private/ShaderCompiler/ShaderCompilerDistributed.cpp b/Engine/Source/Runtime/Engine/Private/ShaderCompiler/ShaderCompilerDistributed.cpp
--- a/Engine/Source/Runtime/Engine/Private/ShaderCompiler/ShaderCompilerDistributed.cpp
+++ b/Engine/Source/Runtime/Engine/Private/ShaderCompiler/ShaderCompilerDistributed.cpp
@@ -258,6 +258,14 @@ void FShaderCompileDistributedThreadRunnable_Interface::DispatchShaderCompileJob
 		FShaderCompileWorkerUtil::WriteTasks(JobsToSerialize, *InputFileAr, Flags, CachedController.SupportsVirtualFiles() ? &TaskCommandData.SourceData : nullptr);
 	}
 
+	if (!TaskCommandData.SourceData.IsEmpty())
+	{
+		// SourceData contains non-owning FMemoryViews into compressed source buffers owned by the compile jobs.
+		// Keep the jobs alive via a refcounted job pointer so the views remain valid for the lifetime of the task.
+		// (handle edge case where distributed build impl is spawning tasks while the owning jobs are being cancelled)
+		TaskCommandData.SourceDataOwner = MakeShared<TArray<FShaderCommonCompileJobPtr>, ESPMode::ThreadSafe>(JobsToSerialize);
+	}
+
 	// Kick off the job
 	NumDispatchedJobs += JobsToSerialize.Num();
 
